UE4 - CrossCompiler tool (replacement for hlslcc.exe)

diff --git a/Engine/Source/Programs/CrossCompilerTool/CrossCompilerTool.Build.cs b/Engine/Source/Programs/CrossCompilerTool/CrossCompilerTool.Build.cs
new file mode 100644
--- /dev/null
+++ b/Engine/Source/Programs/CrossCompilerTool/CrossCompilerTool.Build.cs
@@ -0,0 +1,25 @@
+﻿// Copyright 1998-2014 Epic Games, Inc. All Rights Reserved.
+
+using UnrealBuildTool;
+
+public class CrossCompilerTool : ModuleRules
+{
+	public CrossCompilerTool(TargetInfo Target)
+	{
+		PublicIncludePaths.Add("Runtime/Launch/Public");
+
+		PrivateIncludePaths.Add("Runtime/Launch/Private");		// For LaunchEngineLoop.cpp include
+		PrivateIncludePaths.Add("MetalShaderFormat/Private");		// For Metal includes
+
+		PrivateDependencyModuleNames.AddRange(new string []
+			{
+				"Core",
+				"Projects",
+				"MetalShaderFormat"
+			});
+
+		AddThirdPartyPrivateStaticDependencies(Target,
+			"HLSLCC"
+		);
+	}
+}
diff --git a/Engine/Source/Programs/CrossCompilerTool/CrossCompilerTool.Target.cs b/Engine/Source/Programs/CrossCompilerTool/CrossCompilerTool.Target.cs
new file mode 100644
--- /dev/null
+++ b/Engine/Source/Programs/CrossCompilerTool/CrossCompilerTool.Target.cs
@@ -0,0 +1,59 @@
+﻿// Copyright 1998-2014 Epic Games, Inc. All Rights Reserved.
+
+using UnrealBuildTool;
+using System.Collections.Generic;
+
+public class CrossCompilerToolTarget : TargetRules
+{
+	public CrossCompilerToolTarget(TargetInfo Target)
+	{
+		Type = TargetType.Program;
+	}
+
+	//
+	// TargetRules interface.
+	//
+
+	public override void SetupBinaries(
+		TargetInfo Target,
+		ref List<UEBuildBinaryConfiguration> OutBuildBinaryConfigurations,
+		ref List<string> OutExtraModuleNames
+		)
+	{
+		OutBuildBinaryConfigurations.Add(
+			new UEBuildBinaryConfiguration(	InType: UEBuildBinaryType.Executable,
+											InModuleNames: new List<string>() { "CrossCompilerTool" } )
+			);
+	}
+
+	public override bool ShouldCompileMonolithic(UnrealTargetPlatform InPlatform, UnrealTargetConfiguration InConfiguration)
+	{
+		return true;
+	}
+
+	public override void SetupGlobalEnvironment(
+		TargetInfo Target,
+		ref LinkEnvironmentConfiguration OutLinkEnvironmentConfiguration,
+		ref CPPEnvironmentConfiguration OutCPPEnvironmentConfiguration
+		)
+	{
+		// Lean and mean
+		UEBuildConfiguration.bCompileLeanAndMeanUE = true;
+
+		// Never use malloc profiling in Unreal Header Tool.  We set this because often UHT is compiled right before the engine
+		// automatically by Unreal Build Tool, but if bUseMallocProfiler is defined, UHT can operate incorrectly.
+		BuildConfiguration.bUseMallocProfiler = false;
+
+		// No editor needed
+		UEBuildConfiguration.bBuildEditor = false;
+		// Editor-only data, however, is needed
+		UEBuildConfiguration.bBuildWithEditorOnlyData = true;
+
+		// Currently this app is not linking against the engine, so we'll compile out references from Core to the rest of the engine
+		UEBuildConfiguration.bCompileAgainstEngine = false;
+		UEBuildConfiguration.bCompileAgainstCoreUObject = false;
+
+		// UnrealHeaderTool is a console application, not a Windows app (sets entry point to main(), instead of WinMain())
+		OutLinkEnvironmentConfiguration.bIsBuildingConsoleApplication = true;
+	}
+}
diff --git a/Engine/Source/Programs/CrossCompilerTool/Private/CrossCompilerTool.cpp b/Engine/Source/Programs/CrossCompilerTool/Private/CrossCompilerTool.cpp
new file mode 100644
--- /dev/null
+++ b/Engine/Source/Programs/CrossCompilerTool/Private/CrossCompilerTool.cpp
@@ -0,0 +1,107 @@
+// Copyright 1998-2014 Epic Games, Inc. All Rights Reserved.
+
+// CrossCompilerTool.cpp: Driver for testing compilation of an individual shader
+
+#include "CrossCompilerTool.h"
+#include "hlslcc.h"
+#include "metal/MetalBackEnd.h"
+#include "glsl/ir_gen_glsl.h"
+
+#include "RequiredProgramMainCPPInclude.h"
+
+DEFINE_LOG_CATEGORY(LogCrossCompilerTool);
+
+IMPLEMENT_APPLICATION(CrossCompilerTool, "CrossCompilerTool");
+
+namespace CCT
+{
+	static int32 Run(const FRunInfo& RunInfo)
+	{
+
+		ILanguageSpec* Language = nullptr;
+		FCodeBackend* Backend = nullptr;
+
+		int32 Flags = 0;
+
+		Flags |= RunInfo.bRunCPP ? 0 : HLSLCC_NoPreprocess;
+
+		FGlslLanguageSpec GlslLanguage(RunInfo.Target == HCT_FeatureLevelES2);
+		FGlslCodeBackend GlslBackend(Flags);
+		FMetalLanguageSpec MetalLanguage;
+		FMetalCodeBackend MetalBackend(Flags);
+
+		switch (RunInfo.BackEnd)
+		{
+		case CCT::FRunInfo::BE_Metal:
+			Language = &MetalLanguage;
+			Backend = &MetalBackend;
+			break;
+
+		case CCT::FRunInfo::BE_OpenGL:
+			Language = &GlslLanguage;
+			Backend = &GlslBackend;
+			Flags |= HLSLCC_DX11ClipSpace;
+			break;
+
+		default:
+			return 1;
+		}
+
+		FString HLSLShaderSource;
+		if (!FFileHelper::LoadFileToString(HLSLShaderSource, *RunInfo.InputFile))
+		{
+			UE_LOG(LogCrossCompilerTool, Error, TEXT("Couldn't load Input file!"));
+			return 1;
+		}
+
+		ANSICHAR* GLSLShaderSource = 0;
+		ANSICHAR* ErrorLog = 0;
+
+		int Result = HlslCrossCompile(
+			TCHAR_TO_ANSI(*RunInfo.InputFile),
+			TCHAR_TO_ANSI(*HLSLShaderSource),
+			TCHAR_TO_ANSI(*RunInfo.Entry),
+			RunInfo.Frequency,
+			Backend,
+			Language,
+			Flags,
+			RunInfo.Target,
+			&GLSLShaderSource,
+			&ErrorLog
+			);
+
+		free(GLSLShaderSource);
+		free(ErrorLog);
+
+		return 0;
+	}
+}
+
+INT32_MAIN_INT32_ARGC_TCHAR_ARGV()
+{
+	GEngineLoop.PreInit(ArgC, ArgV, TEXT("-NOPACKAGECACHE -Multiprocess"));
+	UE_LOG(LogCrossCompilerTool,Display,TEXT("Hello World"));
+
+	TArray<FString> Tokens, Switches;
+	FCommandLine::Parse(FCommandLine::Get(), Tokens, Switches);
+
+	if (Tokens.Num() < 1)
+	{
+		UE_LOG(LogCrossCompilerTool, Error, TEXT("Missing input file!"));
+		CCT::PrintUsage();
+		return 1;
+	}
+
+	if (Tokens.Num() > 1)
+	{
+		UE_LOG(LogCrossCompilerTool, Warning,TEXT("Ignoring extra command line arguments!"));
+	}
+
+	CCT::FRunInfo RunInfo;
+	if (!RunInfo.Setup(Tokens[0], Switches))
+	{
+		return 1;
+	}
+
+	return CCT::Run(RunInfo);
+}
diff --git a/Engine/Source/Programs/CrossCompilerTool/Private/CrossCompilerTool.h b/Engine/Source/Programs/CrossCompilerTool/Private/CrossCompilerTool.h
new file mode 100644
--- /dev/null
+++ b/Engine/Source/Programs/CrossCompilerTool/Private/CrossCompilerTool.h
@@ -0,0 +1,40 @@
+// Copyright 1998-2014 Epic Games, Inc. All Rights Reserved.
+
+#pragma once
+
+#include "Core.h"
+
+#include "hlslcc.h"
+
+namespace CCT
+{
+	class FRunInfo
+	{
+	public:
+		enum EBackend
+		{
+			BE_Metal,
+			BE_OpenGL,
+			BE_Invalid,
+		};
+		EHlslShaderFrequency Frequency;
+		EHlslCompileTarget Target;
+		FString Entry;
+		FString InputFile;
+		FString OutputFile;
+		EBackend BackEnd;
+		bool bRunCPP;
+		//bool bCPPOnly;
+
+		FRunInfo();
+		bool Setup(const FString& InOptions, const TArray<FString>& InSwitches);
+
+	protected:
+		static EHlslShaderFrequency ParseFrequency(TArray<FString>& InOutSwitches);
+		static EHlslCompileTarget ParseTarget(TArray<FString>& InOutSwitches, EBackend& OutBackEnd);
+	};
+
+	void PrintUsage();
+}
+
+DECLARE_LOG_CATEGORY_EXTERN(LogCrossCompilerTool, Log, All);
diff --git a/Engine/Source/Programs/CrossCompilerTool/Private/CrossCompilerUtils.cpp b/Engine/Source/Programs/CrossCompilerTool/Private/CrossCompilerUtils.cpp
new file mode 100644
--- /dev/null
+++ b/Engine/Source/Programs/CrossCompilerTool/Private/CrossCompilerUtils.cpp
@@ -0,0 +1,257 @@
+// Copyright 1998-2014 Epic Games, Inc. All Rights Reserved.
+
+
+#include "CrossCompilerTool.h"
+
+namespace CCT
+{
+	FRunInfo::FRunInfo() :
+		Frequency(HSF_InvalidFrequency),
+		Target(HCT_InvalidTarget),
+		Entry(""),
+		InputFile(""),
+		OutputFile(""),
+		BackEnd(BE_Invalid),
+		bRunCPP(true)
+	{
+	}
+
+	void PrintUsage()
+	{
+		UE_LOG(LogCrossCompilerTool, Display, TEXT("Usage:"));
+		UE_LOG(LogCrossCompilerTool, Display, TEXT("\tCrossCompilerTool.exe input.hlsl {options}"));
+		UE_LOG(LogCrossCompilerTool, Display, TEXT("\t\tOptions:"));
+		UE_LOG(LogCrossCompilerTool, Display, TEXT("\t\t\t-o=file\tOutput filename"));
+		UE_LOG(LogCrossCompilerTool, Display, TEXT("\t\t\t-entry=function\tMain entry point (defaults to Main())"));
+		//UE_LOG(LogCrossCompilerTool, Display, TEXT("\t\t\t-cpp\tOnly run C preprocessor"));
+		UE_LOG(LogCrossCompilerTool, Display, TEXT("\t\t\t-nocpp\tDo not run C preprocessor"));
+		UE_LOG(LogCrossCompilerTool, Display, TEXT("\t\tProfiles:"));
+		UE_LOG(LogCrossCompilerTool, Display, TEXT("\t\t\t-vs\tCompile as a Vertex Shader"));
+		UE_LOG(LogCrossCompilerTool, Display, TEXT("\t\t\t-ps\tCompile as a Pixel Shader"));
+		UE_LOG(LogCrossCompilerTool, Display, TEXT("\t\t\t-cs\tCompile as a Compute Shader"));
+		UE_LOG(LogCrossCompilerTool, Display, TEXT("\t\t\t-gs\tCompile as a Geomtry Shader"));
+		UE_LOG(LogCrossCompilerTool, Display, TEXT("\t\t\t-hs\tCompile as a Hull Shader"));
+		UE_LOG(LogCrossCompilerTool, Display, TEXT("\t\t\t-ds\tCompile as a Domain Shader"));
+		UE_LOG(LogCrossCompilerTool, Display, TEXT("\t\tTargets:"));
+		UE_LOG(LogCrossCompilerTool, Display, TEXT("\t\t\t-metal\tCompile for Metal"));
+		UE_LOG(LogCrossCompilerTool, Display, TEXT("\t\t\t-es2\tCompile for OpenGL ES 2"));
+		UE_LOG(LogCrossCompilerTool, Display, TEXT("\t\t\t-gl3\tCompile for OpenGL 3.2"));
+		UE_LOG(LogCrossCompilerTool, Display, TEXT("\t\t\t-gl4\tCompile for OpenGL 4.3"));
+	}
+
+	EHlslShaderFrequency FRunInfo::ParseFrequency(TArray<FString>& InOutSwitches)
+	{
+		TArray<FString> OutSwitches;
+		EHlslShaderFrequency Frequency = HSF_InvalidFrequency;
+
+		// Validate profile and target
+		for(FString& Switch : InOutSwitches)
+		{
+			if (Switch == "vs")
+			{
+				if (Frequency != HSF_InvalidFrequency)
+				{
+					UE_LOG(LogCrossCompilerTool, Warning, TEXT("Ignoring extra command line argument -vs!"));
+				}
+				else
+				{
+					Frequency = HSF_VertexShader;
+				}
+			}
+			else if (Switch == "ps")
+			{
+				if (Frequency != HSF_InvalidFrequency)
+				{
+					UE_LOG(LogCrossCompilerTool, Warning, TEXT("Ignoring extra command line argument -ps!"));
+				}
+				else
+				{
+					Frequency = HSF_PixelShader;
+				}
+			}
+			else if (Switch == "cs")
+			{
+				if (Frequency != HSF_InvalidFrequency)
+				{
+					UE_LOG(LogCrossCompilerTool, Warning, TEXT("Ignoring extra command line argument -cs!"));
+				}
+				else
+				{
+					Frequency = HSF_ComputeShader;
+				}
+			}
+			else if (Switch == "gs")
+			{
+				if (Frequency != HSF_InvalidFrequency)
+				{
+					UE_LOG(LogCrossCompilerTool, Warning, TEXT("Ignoring extra command line argument -gs!"));
+				}
+				else
+				{
+					Frequency = HSF_GeometryShader;
+				}
+			}
+			else if (Switch == "hs")
+			{
+				if (Frequency != HSF_InvalidFrequency)
+				{
+					UE_LOG(LogCrossCompilerTool, Warning, TEXT("Ignoring extra command line argument -hs!"));
+				}
+				else
+				{
+					Frequency = HSF_HullShader;
+				}
+			}
+			else if (Switch == "ds")
+			{
+				if (Frequency != HSF_InvalidFrequency)
+				{
+					UE_LOG(LogCrossCompilerTool, Warning, TEXT("Ignoring extra command line argument -ds!"));
+				}
+				else
+				{
+					Frequency = HSF_HullShader;
+				}
+			}
+			else
+			{
+				OutSwitches.Add(Switch);
+			}
+		}
+
+		// Default to PS
+		if (Frequency == HSF_InvalidFrequency)
+		{
+			UE_LOG(LogCrossCompilerTool, Warning, TEXT("No profile given, assuming Pixel Shader (-ps)!"));
+			Frequency = HSF_PixelShader;
+		}
+
+		Swap(InOutSwitches, OutSwitches);
+		return Frequency;
+	}
+
+
+	EHlslCompileTarget FRunInfo::ParseTarget(TArray<FString>& InOutSwitches, EBackend& OutBackEnd)
+	{
+		TArray<FString> OutSwitches;
+		EHlslCompileTarget Target = HCT_InvalidTarget;
+
+		// Validate profile and target
+		for (FString& Switch : InOutSwitches)
+		{
+			if (Switch == "es2")
+			{
+				if (Target != HCT_InvalidTarget)
+				{
+					UE_LOG(LogCrossCompilerTool, Warning, TEXT("Ignoring extra command line argument -es2!"));
+				}
+				else
+				{
+					Target = HCT_FeatureLevelES2;
+					OutBackEnd = BE_OpenGL;
+				}
+			}
+			else if (Switch == "metal")
+			{
+				if (Target != HCT_InvalidTarget)
+				{
+					UE_LOG(LogCrossCompilerTool, Warning, TEXT("Ignoring extra command line argument -metal!"));
+				}
+				else
+				{
+					Target = HCT_FeatureLevelES3_1;
+					OutBackEnd = BE_Metal;
+				}
+			}
+			else if (Switch == "gl3")
+			{
+				if (Target != HCT_InvalidTarget)
+				{
+					UE_LOG(LogCrossCompilerTool, Warning, TEXT("Ignoring extra command line argument -gl3!"));
+				}
+				else
+				{
+					Target = HCT_FeatureLevelSM4;
+					OutBackEnd = BE_OpenGL;
+				}
+			}
+			else if (Switch == "gl4")
+			{
+				if (Target != HCT_InvalidTarget)
+				{
+					UE_LOG(LogCrossCompilerTool, Warning, TEXT("Ignoring extra command line argument -gl4!"));
+				}
+				else
+				{
+					Target = HCT_FeatureLevelSM5;
+					OutBackEnd = BE_OpenGL;
+				}
+			}
+			else
+			{
+				OutSwitches.Add(Switch);
+			}
+		}
+
+		// Defaults to Metal
+		if (Target == HCT_InvalidTarget)
+		{
+			UE_LOG(LogCrossCompilerTool, Warning, TEXT("No profile given, assuming (-gl4)!"));
+			Target = HCT_FeatureLevelSM5;
+			OutBackEnd = BE_OpenGL;
+		}
+		check(OutBackEnd != BE_Invalid);
+
+		Swap(InOutSwitches, OutSwitches);
+		return Target;
+	}
+
+	bool FRunInfo::Setup(const FString& InFilename, const TArray<FString>& InSwitches)
+	{
+		TArray<FString> Switches = InSwitches;
+		Frequency = ParseFrequency(Switches);
+		Target = ParseTarget(Switches, BackEnd);
+		InputFile = InFilename;
+		bRunCPP = true;
+
+		// Now find entry point and output filename
+		for (const FString& Switch : InSwitches)
+		{
+			if (Switch.StartsWith(TEXT("o=")))
+			{
+				if (OutputFile != "")
+				{
+					UE_LOG(LogCrossCompilerTool, Warning, TEXT("Ignoring extra command line argument -o!"));
+				}
+				else
+				{
+					OutputFile = Switch.Right(2);
+				}
+			}
+			else if (Switch.StartsWith(TEXT("entry=")))
+			{
+				if (Entry != TEXT(""))
+				{
+					UE_LOG(LogCrossCompilerTool, Warning, TEXT("Ignoring extra command line argument -entry!"));
+				}
+				else
+				{
+					Entry = Switch.Right(6);
+				}
+			}
+			else if (Switch.StartsWith(TEXT("nocpp")))
+			{
+				bRunCPP = false;
+			}
+		}
+
+		// Default to PS
+		if (Entry == TEXT(""))
+		{
+			UE_LOG(LogCrossCompilerTool, Warning, TEXT("No entry point given, assuming Main (-entry=Main)!"));
+			Target = HCT_FeatureLevelSM5;
+		}
+
+		return true;
+	}
+}
