UE4 - CCT now writes out source

diff --git a/Engine/Source/Developer/ShaderCompilerCommon/Private/ShaderCompilerCommon.cpp b/Engine/Source/Developer/ShaderCompilerCommon/Private/ShaderCompilerCommon.cpp
--- a/Engine/Source/Developer/ShaderCompilerCommon/Private/ShaderCompilerCommon.cpp
+++ b/Engine/Source/Developer/ShaderCompilerCommon/Private/ShaderCompilerCommon.cpp
@@ -183,11 +183,18 @@ bool RemoveUniformBuffersFromSource(FString& SourceCode)
 }
 
 
-FString CreateCrossCompilerBatchFileContents(const FString& ShaderFile, const FString& OutputFile, const FString& FrequencySwitch, const FString& EntryPoint, const FString& VersionSwitch, const FString& ExtraArguments)
+FString CreateCrossCompilerBatchFileContents(const FString& ShaderFile, const FString& OutputFile, const FString& FrequencySwitch, const FString& EntryPoint, const FString& VersionSwitch, bool bUseCCT, const FString& ExtraArguments)
 {
 	FString BatchFile = TEXT("@echo off");
 	BatchFile += TEXT("\nif defined ue.hlslcc GOTO DONE\nset ue.hlslcc=");
-	BatchFile += FPaths::RootDir() / TEXT("Engine\\Source\\ThirdParty\\hlslcc\\hlslcc\\bin\\Win64\\VS2013\\hlslcc_64.exe");
+	if (bUseCCT)
+	{
+		BatchFile += FPaths::RootDir() / TEXT("Engine\\Source\\Binaries\\Win64\\CrossCompilerTool.exe");
+	}
+	else
+	{
+		BatchFile += FPaths::RootDir() / TEXT("Engine\\Source\\ThirdParty\\hlslcc\\hlslcc\\bin\\Win64\\VS2013\\hlslcc_64.exe");
+	}
 	BatchFile += TEXT("\n\n:DONE\n%ue.hlslcc% ");
 	BatchFile += FString::Printf(TEXT("\"%s\" -o=\"%s\" %s -entry=%s %s %s"), *ShaderFile, *OutputFile, *FrequencySwitch, *EntryPoint, *VersionSwitch, *ExtraArguments);
 	BatchFile += TEXT("\npause\n");
diff --git a/Engine/Source/Developer/ShaderCompilerCommon/Public/ShaderCompilerCommon.h b/Engine/Source/Developer/ShaderCompilerCommon/Public/ShaderCompilerCommon.h
--- a/Engine/Source/Developer/ShaderCompilerCommon/Public/ShaderCompilerCommon.h
+++ b/Engine/Source/Developer/ShaderCompilerCommon/Public/ShaderCompilerCommon.h
@@ -70,4 +70,5 @@ extern SHADERCOMPILERCOMMON_API FString CreateCrossCompilerBatchFileContents(
 											const FString& FrequencySwitch,
 											const FString& EntryPoint,
 											const FString& VersionSwitch,
+											bool bUseCCT,
 											const FString& ExtraArguments = TEXT(""));
diff --git a/Engine/Source/Developer/ShaderFormatOpenGL/Private/OpenGLShaderCompiler.cpp b/Engine/Source/Developer/ShaderFormatOpenGL/Private/OpenGLShaderCompiler.cpp
--- a/Engine/Source/Developer/ShaderFormatOpenGL/Private/OpenGLShaderCompiler.cpp
+++ b/Engine/Source/Developer/ShaderFormatOpenGL/Private/OpenGLShaderCompiler.cpp
@@ -1496,7 +1496,7 @@ static FString CreateCrossCompilerBatchFile( const FString& ShaderFile, const FS
 	}
 
 	const TCHAR* ApplyCSE = (CCFlags & HLSLCC_ApplyCommonSubexpressionElimination) != 0 ? TEXT("-cse") : TEXT("");
-	return CreateCrossCompilerBatchFileContents(ShaderFile, OutputFile, FrequencySwitch, EntryPoint, VersionSwitch, ApplyCSE);
+	return CreateCrossCompilerBatchFileContents(ShaderFile, OutputFile, FrequencySwitch, EntryPoint, VersionSwitch, false, ApplyCSE);
 }
 
 /**
diff --git a/Engine/Source/Developer/iOS/MetalShaderFormat/Private/MetalShaderCompiler.cpp b/Engine/Source/Developer/iOS/MetalShaderFormat/Private/MetalShaderCompiler.cpp
--- a/Engine/Source/Developer/iOS/MetalShaderFormat/Private/MetalShaderCompiler.cpp
+++ b/Engine/Source/Developer/iOS/MetalShaderFormat/Private/MetalShaderCompiler.cpp
@@ -811,7 +811,7 @@ static FString CreateCommandLineHLSLCC( const FString& ShaderFile, const FString
 	const TCHAR* VersionSwitch = TEXT("-metal");
 	const TCHAR* FlattenUB = ((CCFlags & HLSLCC_FlattenUniformBuffers) == HLSLCC_FlattenUniformBuffers) ? TEXT("-flattenub") : TEXT("");
 
-	return CreateCrossCompilerBatchFileContents(ShaderFile, OutputFile, FrequencySwitch, EntryPoint, VersionSwitch, FlattenUB);
+	return CreateCrossCompilerBatchFileContents(ShaderFile, OutputFile, FrequencySwitch, EntryPoint, VersionSwitch, true, FlattenUB);
 }
 
 void CompileShader_Metal(const FShaderCompilerInput& Input,FShaderCompilerOutput& Output,const FString& WorkingDirectory)
diff --git a/Engine/Source/Programs/CrossCompilerTool/Private/CrossCompilerTool.cpp b/Engine/Source/Programs/CrossCompilerTool/Private/CrossCompilerTool.cpp
--- a/Engine/Source/Programs/CrossCompilerTool/Private/CrossCompilerTool.cpp
+++ b/Engine/Source/Programs/CrossCompilerTool/Private/CrossCompilerTool.cpp
@@ -4,7 +4,7 @@
 
 #include "CrossCompilerTool.h"
 #include "hlslcc.h"
-#include "metal/MetalBackEnd.h"
+#include "MetalBackEnd.h"
 #include "glsl/ir_gen_glsl.h"
 
 #include "RequiredProgramMainCPPInclude.h"
@@ -17,7 +17,6 @@ namespace CCT
 {
 	static int32 Run(const FRunInfo& RunInfo)
 	{
-
 		ILanguageSpec* Language = nullptr;
 		FCodeBackend* Backend = nullptr;
 
@@ -54,7 +53,7 @@ namespace CCT
 			return 1;
 		}
 
-		ANSICHAR* GLSLShaderSource = 0;
+		ANSICHAR* ShaderSource = 0;
 		ANSICHAR* ErrorLog = 0;
 
 		int Result = HlslCrossCompile(
@@ -66,11 +65,27 @@ namespace CCT
 			Language,
 			Flags,
 			RunInfo.Target,
-			&GLSLShaderSource,
+			&ShaderSource,
 			&ErrorLog
 			);
 
-		free(GLSLShaderSource);
+		if (ErrorLog)
+		{
+			FString OutError(ANSI_TO_TCHAR(ErrorLog));
+			UE_LOG(LogCrossCompilerTool, Warning, TEXT("%s"), *OutError);
+		}
+
+		if (ShaderSource)
+		{
+			FString OutSource(ANSI_TO_TCHAR(ShaderSource));
+			UE_LOG(LogCrossCompilerTool, Display, TEXT("%s"), *OutSource);
+			if (RunInfo.OutputFile.Len() > 0)
+			{
+				FFileHelper::SaveStringToFile(OutSource, *RunInfo.OutputFile);
+			}
+		}
+
+		free(ShaderSource);
 		free(ErrorLog);
 
 		return 0;
@@ -80,7 +95,6 @@ namespace CCT
 INT32_MAIN_INT32_ARGC_TCHAR_ARGV()
 {
 	GEngineLoop.PreInit(ArgC, ArgV, TEXT("-NOPACKAGECACHE -Multiprocess"));
-	UE_LOG(LogCrossCompilerTool,Display,TEXT("Hello World"));
 
 	TArray<FString> Tokens, Switches;
 	FCommandLine::Parse(FCommandLine::Get(), Tokens, Switches);
diff --git a/Engine/Source/Programs/CrossCompilerTool/Private/CrossCompilerUtils.cpp b/Engine/Source/Programs/CrossCompilerTool/Private/CrossCompilerUtils.cpp
--- a/Engine/Source/Programs/CrossCompilerTool/Private/CrossCompilerUtils.cpp
+++ b/Engine/Source/Programs/CrossCompilerTool/Private/CrossCompilerUtils.cpp
@@ -225,7 +225,7 @@ namespace CCT
 				}
 				else
 				{
-					OutputFile = Switch.Right(2);
+					OutputFile = Switch.Mid(2);
 				}
 			}
 			else if (Switch.StartsWith(TEXT("entry=")))
@@ -236,7 +236,7 @@ namespace CCT
 				}
 				else
 				{
-					Entry = Switch.Right(6);
+					Entry = Switch.Mid(6);
 				}
 			}
 			else if (Switch.StartsWith(TEXT("nocpp")))
