Updated the crosscompiler debug tool for ES3.1 Ext

diff --git a/Engine/Source/Programs/CrossCompilerTool/Private/CrossCompilerTool.cpp b/Engine/Source/Programs/CrossCompilerTool/Private/CrossCompilerTool.cpp
--- a/Engine/Source/Programs/CrossCompilerTool/Private/CrossCompilerTool.cpp
+++ b/Engine/Source/Programs/CrossCompilerTool/Private/CrossCompilerTool.cpp
@@ -38,9 +38,9 @@ namespace CCT
 		Flags |= RunInfo.bForcePackedUBs ? (HLSLCC_PackUniforms | HLSLCC_FlattenUniformBufferStructures | HLSLCC_FlattenUniformBuffers) : 0;
 
 		FGlslLanguageSpec GlslLanguage(RunInfo.Target == HCT_FeatureLevelES2);
-		FGlslCodeBackend GlslBackend(Flags);
+		FGlslCodeBackend GlslBackend(Flags, RunInfo.Target);
 		FMetalLanguageSpec MetalLanguage;
-		FMetalCodeBackend MetalBackend(Flags);
+		FMetalCodeBackend MetalBackend(Flags, RunInfo.Target);
 
 		switch (RunInfo.BackEnd)
 		{
diff --git a/Engine/Source/Programs/CrossCompilerTool/Private/CrossCompilerUtils.cpp b/Engine/Source/Programs/CrossCompilerTool/Private/CrossCompilerUtils.cpp
--- a/Engine/Source/Programs/CrossCompilerTool/Private/CrossCompilerUtils.cpp
+++ b/Engine/Source/Programs/CrossCompilerTool/Private/CrossCompilerUtils.cpp
@@ -40,6 +40,7 @@ namespace CCT
 		UE_LOG(LogCrossCompilerTool, Display, TEXT("\t\tTargets:"));
 		UE_LOG(LogCrossCompilerTool, Display, TEXT("\t\t\t-metal\tCompile for Metal"));
 		UE_LOG(LogCrossCompilerTool, Display, TEXT("\t\t\t-es2\tCompile for OpenGL ES 2"));
+		UE_LOG(LogCrossCompilerTool, Display, TEXT("\t\t\t-es31ext\tCompile for OpenGL ES 3.1 with AEP"));
 		UE_LOG(LogCrossCompilerTool, Display, TEXT("\t\t\t-gl3\tCompile for OpenGL 3.2"));
 		UE_LOG(LogCrossCompilerTool, Display, TEXT("\t\t\t-gl4\tCompile for OpenGL 4.3"));
 	}
@@ -156,6 +157,18 @@ namespace CCT
 					OutBackEnd = BE_OpenGL;
 				}
 			}
+			else if (Switch == "es31ext")
+			{
+				if (Target != HCT_InvalidTarget)
+				{
+					UE_LOG(LogCrossCompilerTool, Warning, TEXT("Ignoring extra command line argument -es31ext!"));
+				}
+				else
+				{
+					Target = HCT_FeatureLevelES3_1Ext;
+					OutBackEnd = BE_OpenGL;
+				}
+			}
 			else if (Switch == "metal")
 			{
 				if (Target != HCT_InvalidTarget)
