UE4 - Debug hlslcc

diff --git a/Engine/Source/Programs/CrossCompilerTool/CrossCompilerTool.Target.cs b/Engine/Source/Programs/CrossCompilerTool/CrossCompilerTool.Target.cs
--- a/Engine/Source/Programs/CrossCompilerTool/CrossCompilerTool.Target.cs
+++ b/Engine/Source/Programs/CrossCompilerTool/CrossCompilerTool.Target.cs
@@ -57,5 +57,6 @@ public class CrossCompilerToolTarget : TargetRules
 		OutLinkEnvironmentConfiguration.bIsBuildingConsoleApplication = true;
 
 		//UEBuildConfiguration.bCompileWithStatsWithoutEngine = true;
+		BuildConfiguration.bDebugBuildsActuallyUseDebugCRT = true;
 	}
 }
