UE4 - CCT - Fix some cmd lines (atomic, packed globals into ub)

diff --git a/Engine/Source/Developer/ShaderCompilerCommon/Private/ShaderCompilerCommon.cpp b/Engine/Source/Developer/ShaderCompilerCommon/Private/ShaderCompilerCommon.cpp
--- a/Engine/Source/Developer/ShaderCompilerCommon/Private/ShaderCompilerCommon.cpp
+++ b/Engine/Source/Developer/ShaderCompilerCommon/Private/ShaderCompilerCommon.cpp
@@ -349,7 +349,8 @@ namespace CrossCompiler
 		CCTCmdLine += ((CCFlags & HLSLCC_ApplyCommonSubexpressionElimination) == HLSLCC_ApplyCommonSubexpressionElimination) ? TEXT(" -cse") : TEXT("");
 		CCTCmdLine += ((CCFlags & HLSLCC_ExpandSubexpressions) == HLSLCC_ExpandSubexpressions) ? TEXT(" -xpxpr") : TEXT("");
 		CCTCmdLine += ((CCFlags & HLSLCC_SeparateShaderObjects) == HLSLCC_SeparateShaderObjects) ? TEXT(" -separateshaders") : TEXT("");
-		CCTCmdLine += ((CCFlags & HLSLCC_PackUniformsIntoUniformBuffers) == HLSLCC_PackUniformsIntoUniformBuffers) ? TEXT(" -packglobalsintoub") : TEXT("");
+		CCTCmdLine += ((CCFlags & HLSLCC_PackUniformsIntoUniformBuffers) == HLSLCC_PackUniformsIntoUniformBuffers) ? TEXT(" -packintoubs") : TEXT("");
+		CCTCmdLine += ((CCFlags & HLSLCC_FixAtomicReferences) == HLSLCC_PackUniformsIntoUniformBuffers) ? TEXT(" -fixatomics") : TEXT("");
 		FString BatchFile;
 		if (PLATFORM_MAC)
 		{
diff --git a/Engine/Source/Programs/CrossCompilerTool/CrossCompilerTool.Build.cs b/Engine/Source/Programs/CrossCompilerTool/CrossCompilerTool.Build.cs
--- a/Engine/Source/Programs/CrossCompilerTool/CrossCompilerTool.Build.cs
+++ b/Engine/Source/Programs/CrossCompilerTool/CrossCompilerTool.Build.cs
@@ -9,7 +9,7 @@ public class CrossCompilerTool : ModuleRules
 		PublicIncludePaths.Add("Runtime/Launch/Public");
 
 		PrivateIncludePaths.Add("Runtime/Launch/Private");		// For LaunchEngineLoop.cpp include
-		PrivateIncludePaths.Add("Developer/iOS/MetalShaderFormat/Private");		// For Metal includes
+		PrivateIncludePaths.Add("Developer/Apple/MetalShaderFormat/Private");		// For Metal includes
 		PrivateIncludePaths.Add("Developer/ShaderFormatOpenGL/Private");		// For GLSL includes
 		PrivateIncludePaths.Add("Developer/ShaderCompilerCommon/Private");		// For Lexer includes
 
diff --git a/Engine/Source/Programs/CrossCompilerTool/Private/CrossCompilerTool.cpp b/Engine/Source/Programs/CrossCompilerTool/Private/CrossCompilerTool.cpp
--- a/Engine/Source/Programs/CrossCompilerTool/Private/CrossCompilerTool.cpp
+++ b/Engine/Source/Programs/CrossCompilerTool/Private/CrossCompilerTool.cpp
@@ -45,6 +45,7 @@ namespace CCT
 		Flags |= RunInfo.bGroupFlattenUBs ? HLSLCC_GroupFlattenedUniformBuffers : 0;
 		Flags |= RunInfo.bCSE ? HLSLCC_ApplyCommonSubexpressionElimination : 0;
 		Flags |= RunInfo.bExpandExpressions ? HLSLCC_ExpandSubexpressions : 0;
+		Flags |= RunInfo.bFixAtomics ? HLSLCC_FixAtomicReferences : 0;
 		Flags |= RunInfo.bSeparateShaders ? HLSLCC_SeparateShaderObjects : 0;
 
 		FGlslLanguageSpec GlslLanguage(RunInfo.Target == HCT_FeatureLevelES2);
diff --git a/Engine/Source/Programs/CrossCompilerTool/Private/CrossCompilerTool.h b/Engine/Source/Programs/CrossCompilerTool/Private/CrossCompilerTool.h
--- a/Engine/Source/Programs/CrossCompilerTool/Private/CrossCompilerTool.h
+++ b/Engine/Source/Programs/CrossCompilerTool/Private/CrossCompilerTool.h
@@ -35,6 +35,7 @@ namespace CCT
 		bool bGroupFlattenUBs;
 		bool bCSE;
 		bool bExpandExpressions;
+		bool bFixAtomics;
 		bool bSeparateShaders;
 
 		FRunInfo();
diff --git a/Engine/Source/Programs/CrossCompilerTool/Private/CrossCompilerUtils.cpp b/Engine/Source/Programs/CrossCompilerTool/Private/CrossCompilerUtils.cpp
--- a/Engine/Source/Programs/CrossCompilerTool/Private/CrossCompilerUtils.cpp
+++ b/Engine/Source/Programs/CrossCompilerTool/Private/CrossCompilerUtils.cpp
@@ -24,6 +24,7 @@ namespace CCT
 		bGroupFlattenUBs(false),
 		bCSE(false),
 		bExpandExpressions(false),
+		bFixAtomics(false),
 		bSeparateShaders(false)
 	{
 	}
@@ -31,7 +32,7 @@ namespace CCT
 	void PrintUsage()
 	{
 		UE_LOG(LogCrossCompilerTool, Display, TEXT("Usage:"));
-		UE_LOG(LogCrossCompilerTool, Display, TEXT("\tCrossCompilerTool.exe input.hlsl {options}"));
+		UE_LOG(LogCrossCompilerTool, Display, TEXT("\tCrossCompilerTool.exe input.[usf|hlsl] {options}"));
 		UE_LOG(LogCrossCompilerTool, Display, TEXT("\t\tOptions:"));
 		UE_LOG(LogCrossCompilerTool, Display, TEXT("\t\t\t-o=file\tOutput filename"));
 		UE_LOG(LogCrossCompilerTool, Display, TEXT("\t\t\t-entry=function\tMain entry point (defaults to Main())"));
@@ -47,11 +48,12 @@ namespace CCT
 		UE_LOG(LogCrossCompilerTool, Display, TEXT("\t\t\t-xpxpr\tExpand expressions"));
 		UE_LOG(LogCrossCompilerTool, Display, TEXT("\t\t\t-separateshaders\tUse the separate shaders flags"));
 		UE_LOG(LogCrossCompilerTool, Display, TEXT("\t\t\t-packintoubs\tMove packed global uniform arrays into a uniform buffer"));
+		UE_LOG(LogCrossCompilerTool, Display, TEXT("\t\t\t-fixatomics\tConvert accesses to atomic variable into intrinsics"));
 		UE_LOG(LogCrossCompilerTool, Display, TEXT("\t\tProfiles:"));
 		UE_LOG(LogCrossCompilerTool, Display, TEXT("\t\t\t-vs\tCompile as a Vertex Shader"));
 		UE_LOG(LogCrossCompilerTool, Display, TEXT("\t\t\t-ps\tCompile as a Pixel Shader"));
 		UE_LOG(LogCrossCompilerTool, Display, TEXT("\t\t\t-cs\tCompile as a Compute Shader"));
-		UE_LOG(LogCrossCompilerTool, Display, TEXT("\t\t\t-gs\tCompile as a Geomtry Shader"));
+		UE_LOG(LogCrossCompilerTool, Display, TEXT("\t\t\t-gs\tCompile as a Geometry Shader"));
 		UE_LOG(LogCrossCompilerTool, Display, TEXT("\t\t\t-hs\tCompile as a Hull Shader"));
 		UE_LOG(LogCrossCompilerTool, Display, TEXT("\t\t\t-ds\tCompile as a Domain Shader"));
 		UE_LOG(LogCrossCompilerTool, Display, TEXT("\t\tTargets:"));
@@ -60,7 +62,6 @@ namespace CCT
 		UE_LOG(LogCrossCompilerTool, Display, TEXT("\t\t\t-es31ext\tCompile for OpenGL ES 3.1 with AEP"));
 		UE_LOG(LogCrossCompilerTool, Display, TEXT("\t\t\t-gl3\tCompile for OpenGL 3.2"));
 		UE_LOG(LogCrossCompilerTool, Display, TEXT("\t\t\t-gl4\tCompile for OpenGL 4.3"));
-		UE_LOG(LogCrossCompilerTool, Display, TEXT("\t\t\t-packglobalsintoub"));
 	}
 
 	EHlslShaderFrequency FRunInfo::ParseFrequency(TArray<FString>& InOutSwitches)
@@ -331,6 +332,10 @@ namespace CCT
 			{
 				bPackIntoUBs = true;
 			}
+			else if (Switch.StartsWith(TEXT("fixatomics")))
+			{
+				bFixAtomics = true;
+			}
 			else if (Switch.StartsWith(TEXT("cse")))
 			{
 				bCSE = true;
