fix for CrossCompile build failure on Mac

diff --git a/Engine/Source/Programs/CrossCompilerTool/CrossCompilerTool.Build.cs b/Engine/Source/Programs/CrossCompilerTool/CrossCompilerTool.Build.cs
--- a/Engine/Source/Programs/CrossCompilerTool/CrossCompilerTool.Build.cs
+++ b/Engine/Source/Programs/CrossCompilerTool/CrossCompilerTool.Build.cs
@@ -11,7 +11,6 @@ public class CrossCompilerTool : ModuleRules
 		PrivateIncludePaths.Add("Runtime/Launch/Private");		// For LaunchEngineLoop.cpp include
 		PrivateIncludePaths.Add("Developer/Apple/MetalShaderFormat/Private");		// For Metal includes
 		PrivateIncludePaths.Add("Developer/ShaderFormatOpenGL/Private");		// For GLSL includes
-        PrivateIncludePaths.Add("Developer/VulkanShaderFormat/Private");
         PrivateIncludePaths.Add("Developer/ShaderCompilerCommon/Private");		// For Lexer includes
 
 		PrivateDependencyModuleNames.AddRange(new string []
@@ -21,7 +20,6 @@ public class CrossCompilerTool : ModuleRules
 				"ShaderCompilerCommon",
 				"MetalShaderFormat",
 				"ShaderFormatOpenGL",
-                "VulkanShaderFormat",
 				"ShaderPreprocessor",
 
 				//@todo-rco: Remove me!
@@ -31,5 +29,11 @@ public class CrossCompilerTool : ModuleRules
 		AddEngineThirdPartyPrivateStaticDependencies(Target,
 			"HLSLCC"
 		);
+
+        if (Target.Platform == UnrealTargetPlatform.Win64)
+        {
+            PrivateIncludePaths.Add("Developer/VulkanShaderFormat/Private");
+            PrivateDependencyModuleNames.Add("VulkanShaderFormat");
+        }
 	}
 }
diff --git a/Engine/Source/Programs/CrossCompilerTool/Private/CrossCompilerTool.cpp b/Engine/Source/Programs/CrossCompilerTool/Private/CrossCompilerTool.cpp
--- a/Engine/Source/Programs/CrossCompilerTool/Private/CrossCompilerTool.cpp
+++ b/Engine/Source/Programs/CrossCompilerTool/Private/CrossCompilerTool.cpp
@@ -6,7 +6,9 @@
 #include "hlslcc.h"
 #include "MetalBackend.h"
 #include "GlslBackend.h"
+#if !PLATFORM_MAC
 #include "VulkanBackend.h"
+#endif
 #include "HlslAST.h"
 #include "HlslLexer.h"
 #include "HlslParser.h"
@@ -69,9 +71,11 @@ namespace CCT
 		FGlslCodeBackend GlslBackend(Flags, RunInfo.Target);
 		FMetalLanguageSpec MetalLanguage;
 		FMetalCodeBackend MetalBackend(Flags, CompileTarget);
+#if !PLATFORM_MAC
 		FVulkanBindingTable VulkanBindingTable(RunInfo.Frequency);
 		FVulkanLanguageSpec VulkanLanguage(false);
 		FVulkanCodeBackend VulkanBackend(Flags, VulkanBindingTable, CompileTarget);
+#endif
 
 		switch (RunInfo.BackEnd)
 		{
@@ -87,11 +91,13 @@ namespace CCT
 			Flags |= HLSLCC_DX11ClipSpace;
 			break;
 
+#if !PLATFORM_MAC
 		case CCT::FRunInfo::BE_Vulkan:
 			Language = &VulkanLanguage;
 			Backend = &VulkanBackend;
 			Flags |= HLSLCC_DX11ClipSpace;
 			break;
+#endif
 
 		default:
 			return 1;
