Merging CL 3188523 from //UE4/Dev-Rendering to fix CrossCompilerTool compile issue. Causes problems for users trying to build the whole UE4 solution. Also add to CIS to prevent future breakages.

diff --git a/Engine/Source/Programs/CrossCompilerTool/CrossCompilerTool.Build.cs b/Engine/Source/Programs/CrossCompilerTool/CrossCompilerTool.Build.cs
--- a/Engine/Source/Programs/CrossCompilerTool/CrossCompilerTool.Build.cs
+++ b/Engine/Source/Programs/CrossCompilerTool/CrossCompilerTool.Build.cs
@@ -11,7 +11,8 @@ public class CrossCompilerTool : ModuleRules
 		PrivateIncludePaths.Add("Runtime/Launch/Private");		// For LaunchEngineLoop.cpp include
 		PrivateIncludePaths.Add("Developer/Apple/MetalShaderFormat/Private");		// For Metal includes
 		PrivateIncludePaths.Add("Developer/ShaderFormatOpenGL/Private");		// For GLSL includes
-        PrivateIncludePaths.Add("Developer/ShaderCompilerCommon/Private");		// For Lexer includes
+        PrivateIncludePaths.Add("Developer/ShaderCompilerCommon/Private");      // For Lexer includes
+		PrivateIncludePaths.Add("Runtime/VulkanRHI/Public");      // For Vulkan definitions/includes
 
 		PrivateDependencyModuleNames.AddRange(new string []
 			{
